Add the missing steamshared package file

#rnx

diff --git a/Engine/Plugins/Runtime/Steam/SteamShared/Source/SteamShared/Public/SteamSharedPackage.h b/Engine/Plugins/Runtime/Steam/SteamShared/Source/SteamShared/Public/SteamSharedPackage.h
new file mode 100644
--- /dev/null
+++ b/Engine/Plugins/Runtime/Steam/SteamShared/Source/SteamShared/Public/SteamSharedPackage.h
@@ -0,0 +1,13 @@
+// Copyright 1998-2019 Epic Games, Inc. All Rights Reserved.
+
+// Can't be #pragma once because other modules may define PACKAGE_SCOPE
+
+// Intended to be the last include in an exported class definition
+// Properly defines some members as "public to the module" vs "private to the consumer/user"
+
+#undef PACKAGE_SCOPE
+#ifdef STEAMSHARED_PACKAGE
+#define PACKAGE_SCOPE public
+#else
+#define PACKAGE_SCOPE protected
+#endif
