Include missing modulemanager include in the main steamshared header

#rnx

diff --git a/Engine/Plugins/Runtime/Steam/SteamShared/Source/SteamShared/Public/SteamSharedModule.h b/Engine/Plugins/Runtime/Steam/SteamShared/Source/SteamShared/Public/SteamSharedModule.h
--- a/Engine/Plugins/Runtime/Steam/SteamShared/Source/SteamShared/Public/SteamSharedModule.h
+++ b/Engine/Plugins/Runtime/Steam/SteamShared/Source/SteamShared/Public/SteamSharedModule.h
@@ -5,6 +5,7 @@
 #include "CoreMinimal.h"
 #include "SteamSharedPackage.h"
 #include "Modules/ModuleInterface.h"
+#include "Modules/ModuleManager.h"
 
 #define LOADING_STEAM_CLIENT_LIBRARY_DYNAMICALLY		(PLATFORM_WINDOWS || PLATFORM_MAC || (PLATFORM_LINUX && !IS_MONOLITHIC))
 #define LOADING_STEAM_SERVER_LIBRARY_DYNAMICALLY		(PLATFORM_WINDOWS || (PLATFORM_LINUX && !IS_MONOLITHIC) || PLATFORM_MAC)
