Steam OSS changes for supporting SteamSockets Ping.

Also updates SteamShared to support initializing dedicated servers.

diff --git a/Engine/Plugins/Online/OnlineSubsystemSteam/Source/OnlineSubsystemSteam.Build.cs b/Engine/Plugins/Online/OnlineSubsystemSteam/Source/OnlineSubsystemSteam.Build.cs
--- a/Engine/Plugins/Online/OnlineSubsystemSteam/Source/OnlineSubsystemSteam.Build.cs
+++ b/Engine/Plugins/Online/OnlineSubsystemSteam/Source/OnlineSubsystemSteam.Build.cs
@@ -27,6 +27,7 @@ public class OnlineSubsystemSteam : ModuleRules
 				"OnlineSubsystem",
 				"Json",
 				"PacketHandler",
+				"Projects",
                 "SteamShared"
             }
 		);
diff --git a/Engine/Plugins/Online/OnlineSubsystemSteam/Source/Private/OnlineAsyncTaskManagerSteam.cpp b/Engine/Plugins/Online/OnlineSubsystemSteam/Source/Private/OnlineAsyncTaskManagerSteam.cpp
--- a/Engine/Plugins/Online/OnlineSubsystemSteam/Source/Private/OnlineAsyncTaskManagerSteam.cpp
+++ b/Engine/Plugins/Online/OnlineSubsystemSteam/Source/Private/OnlineAsyncTaskManagerSteam.cpp
@@ -281,7 +281,7 @@ public:
 			if (Session)
 			{
 				// Make sure the session has all the valid session data
-				if (!FillSessionFromLobbyData(LobbyId, *Session) ||
+				if (!FillSessionFromLobbyData(Subsystem, LobbyId, *Session) ||
 					!FillMembersFromLobbyData(LobbyId, *Session))
 				{
 					UE_LOG_ONLINE(Warning, TEXT("Failed to parse session %s lobby update %s"), *Session->SessionName.ToString(), *LobbyId.ToDebugString());
@@ -678,6 +678,11 @@ public:
 			{			
 				SocketSubsystem->FixupSockets(*SessionInt->GameServerSteamId);
 			}
+
+			if (Subsystem)
+			{
+				Subsystem->TriggerOnSteamServerLoginCompletedDelegates(true);
+			}
 		}
 
 		// log on is not finished until OnPolicyResponse() is called
@@ -769,12 +774,66 @@ void FOnlineAsyncTaskManagerSteam::OnSteamServersDisconnectedGS(SteamServersDisc
 	AddToOutQueue(NewEvent);
 }
 
+/**
+ * Notification event from Steam that server login has failed.
+ */
+class FOnlineAsyncEventSteamServerFailedGS : public FOnlineAsyncEvent<FOnlineSubsystemSteam>
+{
+private:
+
+	/** Callback data */
+	SteamServerConnectFailure_t CallbackResults;
+
+	/** Hidden on purpose */
+	FOnlineAsyncEventSteamServerFailedGS()
+	{
+	}
+
+public:
+
+	FOnlineAsyncEventSteamServerFailedGS(FOnlineSubsystemSteam* InSubsystem, SteamServerConnectFailure_t& InResults) :
+		FOnlineAsyncEvent(InSubsystem),
+		CallbackResults(InResults)
+	{
+	}
+
+	virtual ~FOnlineAsyncEventSteamServerFailedGS()
+	{
+	}
+
+	/**
+	 *	Get a human readable description of task
+	 */
+	virtual FString ToString() const override
+	{
+		return FString::Printf(TEXT("FOnlineAsyncEventSteamServerFailedGS Result: %s"), *SteamResultString(CallbackResults.m_eResult));
+	}
+
+	/**
+	 * Give the async task a chance to marshal its data back to the game thread
+	 * Can only be called on the game thread by the async task manager
+	 */
+	virtual void Finalize() override
+	{
+		if (Subsystem)
+		{
+			Subsystem->TriggerOnSteamServerLoginCompletedDelegates(false);
+		}
+	}
+};
+
 /**
  *	GameServer API version of disconnected from Steam backend callback
  */
 void FOnlineAsyncTaskManagerSteam::OnSteamServersConnectFailureGS(SteamServerConnectFailure_t* CallbackData)
 {
-	UE_LOG_ONLINE(Warning, TEXT("Steam connection failure."));
+	// Only do something if we are no longer retrying to connect with the backend.
+	if (!CallbackData->m_bStillRetrying)
+	{
+		FOnlineAsyncEventSteamServerFailedGS* NewEvent = new FOnlineAsyncEventSteamServerFailedGS(SteamSubsystem, *CallbackData);
+		UE_LOG_ONLINE(Verbose, TEXT("%s"), *NewEvent->ToString());
+		AddToOutQueue(NewEvent);
+	}
 }
 
 /**
diff --git a/Engine/Plugins/Online/OnlineSubsystemSteam/Source/Private/OnlineSessionAsyncLobbySteam.cpp b/Engine/Plugins/Online/OnlineSubsystemSteam/Source/Private/OnlineSessionAsyncLobbySteam.cpp
--- a/Engine/Plugins/Online/OnlineSubsystemSteam/Source/Private/OnlineSessionAsyncLobbySteam.cpp
+++ b/Engine/Plugins/Online/OnlineSubsystemSteam/Source/Private/OnlineSessionAsyncLobbySteam.cpp
@@ -3,6 +3,8 @@
 #include "OnlineSessionAsyncLobbySteam.h"
 #include "SocketSubsystem.h"
 #include "OnlineSessionInterfaceSteam.h"
+#include "OnlineSubsystemSteam.h"
+#include "OnlinePingInterfaceSteam.h"
 #include "IPAddressSteam.h"
 #include "SteamSessionKeys.h"
 #include "SteamUtilities.h"
@@ -191,6 +193,18 @@ static void GetLobbyKeyValuePairsFromSession(const FOnlineSession* Session, FSte
 		GetLobbyKeyValuePairsFromSessionInfo(*SessionInfo, KeyValuePairs);
 	}
 
+	// Write in the ping information if the host is using SteamSockets
+	FOnlineSubsystemSteam* SteamSubsystem = static_cast<FOnlineSubsystemSteam*>(IOnlineSubsystem::Get(STEAM_SUBSYSTEM));
+	if (SteamSubsystem && SteamSubsystem->GetPingInterface().IsValid() && SteamSubsystem->GetPingInterface()->IsUsingP2PRelays())
+	{
+		FString PingData = SteamSubsystem->GetPingInterface()->GetHostPingData();
+		// Don't bother writing empty data.
+		if (!PingData.IsEmpty())
+		{
+			KeyValuePairs.Add(STEAMKEY_P2PPING, PingData);
+		}
+	}
+
 	GetLobbyKeyValuePairsFromSessionSettings(Session->SessionSettings, KeyValuePairs);
 }
 
@@ -198,12 +212,14 @@ static void GetLobbyKeyValuePairsFromSession(const FOnlineSession* Session, FSte
  *	Populate an FSession data structure from the data stored with a lobby
  * Expects a certain number of keys to be present otherwise this will fail
  *
+ * @param SteamSubsystem the online subsystem to use
  * @param LobbyId the Steam lobby to fill data from
  * @param Session empty session structure to fill in
+ * @param SearchData if this is provided and the Steam session supports SteamSockets, the ping for the lobby will be written into the search result data.
  *
  * @return true if successful, false otherwise
  */
-bool FillSessionFromLobbyData(FUniqueNetIdSteam& LobbyId, FOnlineSession& Session)
+bool FillSessionFromLobbyData(FOnlineSubsystemSteam* SteamSubsystem, FUniqueNetIdSteam& LobbyId, FOnlineSession& Session, FOnlineSessionSearchResult* SearchData)
 {
 	bool bSuccess = true;
 
@@ -223,7 +239,7 @@ bool FillSessionFromLobbyData(FUniqueNetIdSteam& LobbyId, FOnlineSession& Sessio
 	int32 HostKeysFound = 0;
 	int32 SteamAddrKeysFound = 0;
 
-	const int32 LobbyDataBufferSize = 1024;
+	const int32 LobbyDataBufferSize = 2000; // The Limit on Steam is 8192, but we won't use that much.
 	ANSICHAR Key[LobbyDataBufferSize];
 	ANSICHAR Value[LobbyDataBufferSize];
 
@@ -296,6 +312,18 @@ bool FillSessionFromLobbyData(FUniqueNetIdSteam& LobbyId, FOnlineSession& Sessio
 			Session.NumOpenPublicConnections = FCString::Atoi(ANSI_TO_TCHAR(Value));
 			KeysFound++;
 		}
+		else if (FCStringAnsi::Strcmp(Key, STEAMKEY_P2PPING) == 0)
+		{
+			if (SteamSubsystem && SearchData != nullptr && SteamSubsystem->GetPingInterface().IsValid())
+			{
+				int32 PingResult = SteamSubsystem->GetPingInterface()->GetPingFromHostData(ANSI_TO_TCHAR(Value));
+				// If the ping was successful, we'll want to write the data. Otherwise let it stay at the default value (of 999)
+				if (PingResult >= 0)
+				{
+					SearchData->PingInMs = PingResult;
+				}
+			}
+		}
 		else if (FCStringAnsi::Stricmp(Key, STEAMKEY_HOSTIP) == 0)
 		{
 			uint32 HostIp = FCString::Atoi(ANSI_TO_TCHAR(Value));
@@ -884,7 +912,7 @@ void FOnlineAsyncTaskSteamFindLobbiesBase::CreateQuery()
 void FOnlineAsyncTaskSteamFindLobbiesBase::ParseSearchResult(FUniqueNetIdSteam& LobbyId)
 {
 	FOnlineSessionSearchResult* NewSearchResult = new (SearchSettings->SearchResults) FOnlineSessionSearchResult();
-	if (!FillSessionFromLobbyData(LobbyId, NewSearchResult->Session))
+	if (!FillSessionFromLobbyData(Subsystem, LobbyId, NewSearchResult->Session, NewSearchResult))
 	{
 		UE_LOG_ONLINE_SESSION(Warning, TEXT("Unable to parse search result for lobby '%s'"), *LobbyId.ToDebugString());
 		// Remove the failed element
diff --git a/Engine/Plugins/Online/OnlineSubsystemSteam/Source/Private/OnlineSessionAsyncLobbySteam.h b/Engine/Plugins/Online/OnlineSubsystemSteam/Source/Private/OnlineSessionAsyncLobbySteam.h
--- a/Engine/Plugins/Online/OnlineSubsystemSteam/Source/Private/OnlineSessionAsyncLobbySteam.h
+++ b/Engine/Plugins/Online/OnlineSubsystemSteam/Source/Private/OnlineSessionAsyncLobbySteam.h
@@ -445,12 +445,14 @@ ELobbyType BuildLobbyType(FOnlineSessionSettings* SessionSettings);
  *	Populate an FSession data structure from the data stored with a lobby
  * Expects a certain number of keys to be present otherwise this will fail
  *
+ * @param SteamSubsystem the online subsystem to use
  * @param LobbyId the Steam lobby to fill data from
  * @param Session empty session structure to fill in
+ * @param SearchData if this is provided and the Steam session supports SteamSockets, the ping for the lobby will be written into the search result data.
  *
  * @return true if successful, false otherwise
  */
-bool FillSessionFromLobbyData(FUniqueNetIdSteam& LobbyId, class FOnlineSession& Session);
+bool FillSessionFromLobbyData(FOnlineSubsystemSteam* SteamSubsystem, FUniqueNetIdSteam& LobbyId, class FOnlineSession& Session, FOnlineSessionSearchResult* SearchData = nullptr);
 
 /**
  *	Populate an FSession data structure from the data stored with the members of the lobby
diff --git a/Engine/Plugins/Online/OnlineSubsystemSteam/Source/Private/OnlineSubsystemSteam.cpp b/Engine/Plugins/Online/OnlineSubsystemSteam/Source/Private/OnlineSubsystemSteam.cpp
--- a/Engine/Plugins/Online/OnlineSubsystemSteam/Source/Private/OnlineSubsystemSteam.cpp
+++ b/Engine/Plugins/Online/OnlineSubsystemSteam/Source/Private/OnlineSubsystemSteam.cpp
@@ -8,6 +8,7 @@
 #include "HAL/RunnableThread.h"
 #include "Misc/ScopeLock.h"
 #include "Serialization/BufferArchive.h"
+#include "Interfaces/IPluginManager.h"
 #include "SocketSubsystem.h"
 
 #include "IPAddress.h"
@@ -231,6 +232,14 @@ FOnlineAuthSteamUtilsPtr FOnlineSubsystemSteam::GetAuthInterfaceUtils() const
 	return AuthInterfaceUtils;
 }
 
+FOnlinePingSteamPtr FOnlineSubsystemSteam::GetPingInterface() const
+{
+	return PingInterface;
+}
+void FOnlineSubsystemSteam::SetPingInterface(FOnlinePingSteamPtr InPingInterface)
+{
+	PingInterface = InPingInterface;
+}
 IOnlineSessionPtr FOnlineSubsystemSteam::GetSessionInterface() const
 {
 	return SessionInterface;
@@ -432,7 +441,13 @@ bool FOnlineSubsystemSteam::Init()
 
 	if (bClientInitSuccess && bServerInitSuccess)
 	{
-		CreateSteamSocketSubsystem();
+		TSharedPtr<IPlugin> SocketsPlugin = IPluginManager::Get().FindPlugin(TEXT("SteamSockets"));
+		if (!SocketsPlugin.IsValid() || (SocketsPlugin.IsValid() && !SocketsPlugin->IsEnabled()))
+		{
+			UE_LOG_ONLINE(Log, TEXT("Initializing SteamNetworking Layer"));
+			CreateSteamSocketSubsystem();
+			bUsingSteamNetworking = true;
+		}
 
 		// Create the online async task thread
 		OnlineAsyncTaskThreadRunnable = new FOnlineAsyncTaskManagerSteam(this);
@@ -524,6 +539,7 @@ bool FOnlineSubsystemSteam::Shutdown()
 	DESTRUCT_INTERFACE(IdentityInterface);
 	DESTRUCT_INTERFACE(AuthInterface);
 	DESTRUCT_INTERFACE(AuthInterfaceUtils);
+	DESTRUCT_INTERFACE(PingInterface);
 	DESTRUCT_INTERFACE(SessionInterface);
 	DESTRUCT_INTERFACE(PresenceInterface);
 
@@ -531,7 +547,10 @@ bool FOnlineSubsystemSteam::Shutdown()
 
 	ClearUserCloudFiles();
 
-	DestroySteamSocketSubsystem();
+	if (bUsingSteamNetworking)
+	{
+		DestroySteamSocketSubsystem();
+	}
 
 	ShutdownSteamworks();
 
@@ -592,7 +611,7 @@ bool FOnlineSubsystemSteam::InitSteamworksClient(bool bRelaunchInSteam, int32 St
 	// Otherwise initialize as normal
 	else
 	{
-		SteamAPIClientHandle = FSteamSharedModule::Get().ObtainSteamInstanceHandle();
+		SteamAPIClientHandle = FSteamSharedModule::Get().ObtainSteamClientInstanceHandle();
 		// Steamworks needs to initialize as close to start as possible, so it can hook its overlay into Direct3D, etc.
 		bSteamworksClientInitialized = (SteamAPIClientHandle.IsValid() ? true : false);
 
@@ -658,56 +677,16 @@ bool FOnlineSubsystemSteam::InitSteamworksClient(bool bRelaunchInSteam, int32 St
 
 bool FOnlineSubsystemSteam::InitSteamworksServer()
 {
-	bSteamworksGameServerInitialized = false;
-
-	// Initialize the Steam game server interfaces (done regardless of whether or not a server will be setup)
-	// NOTE: The port values specified here, are not changeable once the interface is setup
-	
-	uint32 LocalServerIP = 0;
-	FString MultiHome;
-	if (FParse::Value(FCommandLine::Get(), TEXT("MULTIHOME="), MultiHome) && !MultiHome.IsEmpty())
-	{
-		TSharedPtr<FInternetAddr> MultiHomeIP = ISocketSubsystem::Get(PLATFORM_SOCKETSUBSYSTEM)->GetAddressFromString(MultiHome);
-		if (MultiHomeIP.IsValid())
-		{
-			MultiHomeIP->GetIp(LocalServerIP);
-		}
-	}
-
-	if (FParse::Value(FCommandLine::Get(), TEXT("Port="), GameServerGamePort) == false)
-	{
-		GConfig->GetInt(TEXT("URL"), TEXT("Port"), GameServerGamePort, GEngineIni);
-	}
-	
-	GameServerSteamPort = GameServerGamePort + 1;
-
-	// Allow the command line to override the default query port
-	if (FParse::Value(FCommandLine::Get(), TEXT("QueryPort="), GameServerQueryPort) == false)
-	{
-		if (!GConfig->GetInt(TEXT("OnlineSubsystemSteam"), TEXT("GameServerQueryPort"), GameServerQueryPort, GEngineIni))
-		{
-			GameServerQueryPort = 27015;
-		}
-	}
-
-	bool bVACEnabled = false;
-	GConfig->GetBool(TEXT("OnlineSubsystemSteam"), TEXT("bVACEnabled"), bVACEnabled, GEngineIni);
-
-	FString GameVersion;
-	GConfig->GetString(TEXT("OnlineSubsystemSteam"), TEXT("GameVersion"), GameVersion, GEngineIni);
-	if (GameVersion.Len() == 0)
-	{
-		UE_LOG_ONLINE(Warning, TEXT("[OnlineSubsystemSteam].GameVersion is not set. Server advertising will fail"));
-	}
-
-	// NOTE: IP of 0 causes SteamGameServer_Init to automatically use the public (external) IP
-	UE_LOG_ONLINE(Verbose, TEXT("Initializing Steam Game Server IP: 0x%08X Port: %d SteamPort: %d QueryPort: %d"), LocalServerIP, GameServerGamePort, GameServerSteamPort, GameServerQueryPort);
-	bSteamworksGameServerInitialized = SteamGameServer_Init(LocalServerIP, GameServerSteamPort, GameServerGamePort, GameServerQueryPort,
-		(bVACEnabled ? eServerModeAuthenticationAndSecure : eServerModeAuthentication),
-		TCHAR_TO_UTF8(*GameVersion));
+	SteamAPIServerHandle = FSteamSharedModule::Get().ObtainSteamServerInstanceHandle();
+	bSteamworksGameServerInitialized = (SteamAPIServerHandle.IsValid());
 
 	if (bSteamworksGameServerInitialized)
 	{
+		// Grab the port values so that we save them.
+		GameServerGamePort = SteamAPIServerHandle->GetGamePort();
+		GameServerSteamPort = SteamAPIServerHandle->GetSteamPort();
+		GameServerQueryPort = SteamAPIServerHandle->GetQueryPort();
+
 		// Test all the Steam interfaces
 		#define GET_STEAMWORKS_INTERFACE(Interface) \
 		if (Interface() == nullptr) \
@@ -739,21 +718,13 @@ void FOnlineSubsystemSteam::ShutdownSteamworks()
 {
 	if (bSteamworksGameServerInitialized)
 	{
-		if (SteamGameServer() != nullptr)
+		SteamAPIServerHandle.Reset();
+		if (SessionInterface.IsValid())
 		{
-			// Since SteamSDK 1.17, LogOff is required to stop the game server advertising after exit; ensure we don't miss this at shutdown
-			if (SteamGameServer()->BLoggedOn())
-			{
-				SteamGameServer()->LogOff();
-			}
-
-			SteamGameServer_Shutdown();
-			if (SessionInterface.IsValid())
-			{
-				SessionInterface->GameServerSteamId = nullptr;
-				SessionInterface->bSteamworksGameServerConnected = false;
-			}	
+			SessionInterface->GameServerSteamId = nullptr;
+			SessionInterface->bSteamworksGameServerConnected = false;
 		}
+		bSteamworksGameServerInitialized = false;
 	}
 
 	if (bSteamworksClientInitialized)
diff --git a/Engine/Plugins/Online/OnlineSubsystemSteam/Source/Private/SteamSessionKeys.h b/Engine/Plugins/Online/OnlineSubsystemSteam/Source/Private/SteamSessionKeys.h
--- a/Engine/Plugins/Online/OnlineSubsystemSteam/Source/Private/SteamSessionKeys.h
+++ b/Engine/Plugins/Online/OnlineSubsystemSteam/Source/Private/SteamSessionKeys.h
@@ -27,6 +27,7 @@
 #define STEAMKEY_HOSTPORT "HOSTPORT"
 #define STEAMKEY_P2PADDR "P2PADDR"
 #define STEAMKEY_P2PPORT "P2PPORT"
+#define STEAMKEY_P2PPING "P2PPING"
 
 /**
  * Converts an engine key and its data type to an appropriate Steam key for use with lobbies/gameservers
diff --git a/Engine/Plugins/Online/OnlineSubsystemSteam/Source/Public/OnlinePingInterfaceSteam.h b/Engine/Plugins/Online/OnlineSubsystemSteam/Source/Public/OnlinePingInterfaceSteam.h
new file mode 100644
--- /dev/null
+++ b/Engine/Plugins/Online/OnlineSubsystemSteam/Source/Public/OnlinePingInterfaceSteam.h
@@ -0,0 +1,64 @@
+// Copyright 1998-2019 Epic Games, Inc. All Rights Reserved.
+
+#pragma once
+
+#include "CoreMinimal.h"
+
+/**
+ * An abstract ping interface for Steam ideally usable with the new 
+ * SteamSocket internet protocol
+ */
+class ONLINESUBSYSTEMSTEAM_API FOnlinePingInterfaceSteam
+{
+public:
+	FOnlinePingInterfaceSteam(class FOnlineSubsystemSteam* InSubsystem) :
+		Subsystem(InSubsystem)
+	{
+	}
+
+	virtual ~FOnlinePingInterfaceSteam()
+	{
+	}
+	
+	/**
+	 * Returns if this application is set up to use the Steam P2P Relay Network
+	 * for communication.
+	 *
+	 * Uses OnlineSubsystemSteam.bAllowP2PPacketRelay
+	 *
+	 * @return true if relays are enabled for P2P connections.
+	 */
+	virtual bool IsUsingP2PRelays() const = 0;
+
+	/**
+	 * Returns the P2P relay ping information for the current machine. This information can be
+	 * serialized over the network and used to calculate the ping data between a client and a host.
+	 *
+	 * @return relay information blob stored as a string for relaying over the network. 
+	 *         If an error occurred, the return is an empty string.
+	 */
+	virtual FString GetHostPingData() const = 0;
+
+	/**
+	 * Calculates the ping of this client using the given host's ping data obtained from GetHostPingData.
+	 * 
+	 * @param HostPingStr The relay information blob we got from the host. This information
+	 *                    should be directly serialized over the network and not tampered with.
+	 *
+	 * @return The ping value to the given host if it can be calculated, otherwise -1 on error.
+	 */
+	virtual int32 GetPingFromHostData(const FString& HostPingStr) const = 0;
+
+	/**
+	* An informative member that allows us to check if we are recalculating our ping
+	* information over the Valve network. Data is additively modified during recalculation
+	* such that we do not need to block on this function returning false before using
+	* ping data.
+	*
+	* @return true if we're recalculating our ping within the Valve relay network.
+	*/
+	virtual bool IsRecalculatingPing() const = 0;
+	
+protected:
+	FOnlineSubsystemSteam* Subsystem;
+};
\ No newline at end of file
diff --git a/Engine/Plugins/Online/OnlineSubsystemSteam/Source/Public/OnlineSubsystemSteam.h b/Engine/Plugins/Online/OnlineSubsystemSteam/Source/Public/OnlineSubsystemSteam.h
--- a/Engine/Plugins/Online/OnlineSubsystemSteam/Source/Public/OnlineSubsystemSteam.h
+++ b/Engine/Plugins/Online/OnlineSubsystemSteam/Source/Public/OnlineSubsystemSteam.h
@@ -3,6 +3,7 @@
 #pragma once
 
 #include "CoreMinimal.h"
+#include "OnlineDelegateMacros.h"
 #include "OnlineSubsystemImpl.h"
 #include "OnlineSubsystemSteamPackage.h"
 
@@ -18,6 +19,7 @@ class FOnlineVoiceSteam;
 class FOnlinePresenceSteam;
 class FOnlineAuthSteam;
 class FOnlineAuthUtilsSteam;
+class FOnlinePingInterfaceSteam;
 
 /** Forward declarations of all interface classes */
 typedef TSharedPtr<class FOnlineSessionSteam, ESPMode::ThreadSafe> FOnlineSessionSteamPtr;
@@ -32,6 +34,16 @@ typedef TSharedPtr<class FOnlineAchievementsSteam, ESPMode::ThreadSafe> FOnlineA
 typedef TSharedPtr<class FOnlinePresenceSteam, ESPMode::ThreadSafe> FOnlinePresenceSteamPtr;
 typedef TSharedPtr<class FOnlineAuthSteam, ESPMode::ThreadSafe> FOnlineAuthSteamPtr;
 typedef TSharedPtr<class FOnlineAuthUtilsSteam, ESPMode::ThreadSafe> FOnlineAuthSteamUtilsPtr;
+typedef TSharedPtr<class FOnlinePingInterfaceSteam, ESPMode::ThreadSafe> FOnlinePingSteamPtr;
+
+/**
+* Delegate fired when a Steam Game Server has completed its login tasks with the Steam backend.
+*
+* @param bWasSuccessful if the login completed successfully
+*/
+DECLARE_MULTICAST_DELEGATE_OneParam(FOnSteamServerLoginCompleted, bool /* bWasSuccessful */);
+typedef FOnSteamServerLoginCompleted::FDelegate FOnSteamServerLoginCompletedDelegate;
+
 
 /**
  *	OnlineSubsystemSteam - Implementation of the online subsystem for STEAM services
@@ -47,6 +59,9 @@ protected:
 	/** Whether or not the Steam game server API is initialized */
 	bool bSteamworksGameServerInitialized;
 
+	/** If we are using the SteamNetworking protocol or not. */
+	bool bUsingSteamNetworking;
+
 	/** Steam App ID for the running game */
 	uint32 SteamAppID;
 
@@ -99,6 +114,9 @@ protected:
 	FOnlineAuthSteamPtr AuthInterface;
 	FOnlineAuthSteamUtilsPtr AuthInterfaceUtils;
 
+	/** Interface for dynamically calculating SteamNetworking ping based off protocol */
+	FOnlinePingSteamPtr PingInterface;
+
 	/** Online async task runnable */
 	class FOnlineAsyncTaskManagerSteam* OnlineAsyncTaskThreadRunnable;
 
@@ -106,7 +124,10 @@ protected:
 	class FRunnableThread* OnlineAsyncTaskThread;
 
 	/** Steam Client API Handle */
-	TSharedPtr<class FSteamInstanceHandler> SteamAPIClientHandle;
+	TSharedPtr<class FSteamClientInstanceHandler> SteamAPIClientHandle;
+
+	/** Steam Server API Handle */
+	TSharedPtr<class FSteamServerInstanceHandler> SteamAPIServerHandle;
 
 PACKAGE_SCOPE:
 
@@ -116,6 +137,7 @@ PACKAGE_SCOPE:
 		FOnlineSubsystemImpl(STEAM_SUBSYSTEM, InInstanceName),
 		bSteamworksClientInitialized(false),
 		bSteamworksGameServerInitialized(false),
+		bUsingSteamNetworking(false),
 		SteamAppID(0),
 		GameServerSteamPort(0),
 		GameServerGamePort(0),
@@ -132,9 +154,11 @@ PACKAGE_SCOPE:
 		PresenceInterface(nullptr),
 		AuthInterface(nullptr),
 		AuthInterfaceUtils(nullptr),
+		PingInterface(nullptr),
 		OnlineAsyncTaskThreadRunnable(nullptr),
 		OnlineAsyncTaskThread(nullptr),
-		SteamAPIClientHandle(nullptr)
+		SteamAPIClientHandle(nullptr),
+		SteamAPIServerHandle(nullptr)
 	{}
 
 	/** Critical sections for thread safe operation of the cloud files */
@@ -214,6 +238,9 @@ public:
 
 	virtual FOnlineAuthSteamPtr GetAuthInterface() const;
 	virtual FOnlineAuthSteamUtilsPtr GetAuthInterfaceUtils() const;
+	
+	virtual FOnlinePingSteamPtr GetPingInterface() const;
+	virtual void SetPingInterface(FOnlinePingSteamPtr InPingInterface);
 
 	// IOnlineSubsystem
 
@@ -313,6 +340,12 @@ public:
 	{
 		return GameServerQueryPort;
 	}
+
+	/**
+	 * This delegate fires whenever a steam login has succeeded or failed its async task.
+	 * Useful for modules that need to check to see if a user is logged in before running other behavior
+	 */
+	DEFINE_ONLINE_DELEGATE_ONE_PARAM(OnSteamServerLoginCompleted, bool);
 };
 
 namespace FNetworkProtocolTypes
diff --git a/Engine/Plugins/Runtime/Steam/SteamController/Source/SteamController/Private/SteamController.cpp b/Engine/Plugins/Runtime/Steam/SteamController/Source/SteamController/Private/SteamController.cpp
--- a/Engine/Plugins/Runtime/Steam/SteamController/Source/SteamController/Private/SteamController.cpp
+++ b/Engine/Plugins/Runtime/Steam/SteamController/Source/SteamController/Private/SteamController.cpp
@@ -40,7 +40,7 @@ public:
 		GConfig->GetDouble(TEXT("/Script/Engine.InputSettings"), TEXT("ButtonRepeatDelay"), ButtonRepeatDelay, GInputIni);
 
 		// Initialize the API, so we can start calling SteamController functions
-		SteamAPIHandle = FSteamSharedModule::Get().ObtainSteamInstanceHandle();
+		SteamAPIHandle = FSteamSharedModule::Get().ObtainSteamClientInstanceHandle();
 
 		// [RCL] 2015-01-23 FIXME: move to some other code than constructor so we can handle failures more gracefully
 		if (SteamAPIHandle.IsValid() && (SteamInput() != nullptr))
@@ -330,7 +330,7 @@ private:
 	TSharedRef<FGenericApplicationMessageHandler> MessageHandler;
 
 	/** SteamAPI initialized **/
-	TSharedPtr<class FSteamInstanceHandler> SteamAPIHandle;
+	TSharedPtr<class FSteamClientInstanceHandler> SteamAPIHandle;
 
 	/** SteamController initialized **/
 	bool bSteamControllerInitialized;
diff --git a/Engine/Plugins/Runtime/Steam/SteamShared/Source/SteamShared/Private/SteamSharedModule.cpp b/Engine/Plugins/Runtime/Steam/SteamShared/Source/SteamShared/Private/SteamSharedModule.cpp
--- a/Engine/Plugins/Runtime/Steam/SteamShared/Source/SteamShared/Private/SteamSharedModule.cpp
+++ b/Engine/Plugins/Runtime/Steam/SteamShared/Source/SteamShared/Private/SteamSharedModule.cpp
@@ -6,6 +6,9 @@
 #include "Misc/CommandLine.h"
 #include "Modules/ModuleManager.h"
 #include "HAL/PlatformProcess.h"
+#include "Misc/ConfigCacheIni.h"
+#include "SocketSubsystem.h"
+#include "IPAddress.h"
 
 #ifndef STEAM_SDK_INSTALLED
 #error Steam SDK not located! Expected to be found in Engine/Source/ThirdParty/Steamworks/{SteamVersion}
@@ -14,6 +17,7 @@
 // Steam API for Initialization
 THIRD_PARTY_INCLUDES_START
 #include "steam/steam_api.h"
+#include "steam/steam_gameserver.h"
 THIRD_PARTY_INCLUDES_END
 
 DEFINE_LOG_CATEGORY_STATIC(LogSteamShared, Log, All);
@@ -48,6 +52,18 @@ FString FSteamSharedModule::GetSteamModulePath() const
 #endif	//PLATFORM_WINDOWS
 }
 
+bool FSteamSharedModule::CanLoadClientDllsOnServer() const
+{
+#if PLATFORM_WINDOWS
+	if (IsRunningDedicatedServer())
+	{
+		return true;
+	}
+#endif
+
+	return false;
+}
+
 void FSteamSharedModule::StartupModule()
 {
 	// On startup load the modules. Anyone who uses this shared library can guarantee
@@ -61,42 +77,77 @@ void FSteamSharedModule::ShutdownModule()
 	if (AreSteamDllsLoaded())
 	{
 		// Make sure everyone cleaned up their instances properly.
-		if (InstanceHandlerObserver.IsValid()) // If our weakptr is still valid, that means there are cases that didn't clean up somehow.
+		if (SteamClientObserver.IsValid()) // If our weakptr is still valid, that means there are cases that didn't clean up somehow.
 		{
 			// If they have not, warn to them they need to clean up in the future and force the deletion.
-			uint32 NumSharedReferences = InstanceHandlerObserver.Pin().GetSharedReferenceCount();
+			uint32 NumSharedReferences = SteamClientObserver.Pin().GetSharedReferenceCount();
 			// Make sure to subtract 1 here as we just created a new sharedptr in order to get the reference count
 			// (this sharedptr is out of scope so it does not matter)
 			UE_LOG(LogSteamShared, Warning, TEXT("There are still %d additional Steam instances in use. These must be shutdown before unloading the module!"), NumSharedReferences - 1);
 		}
-		InstanceHandlerObserver.Reset(); // Force the clearing of our weakptr.
+		SteamClientObserver.Reset(); // Force the clearing of our weakptr.
+
+		// Do the above but this time with the server observer.
+		if (SteamServerObserver.IsValid())
+		{
+			uint32 NumSharedReferences = SteamServerObserver.Pin().GetSharedReferenceCount();
+			UE_LOG(LogSteamShared, Warning, TEXT("There are still %d additional Steam server instances in use. These must be shutdown before unloading the module!"), NumSharedReferences - 1);
+		}
+		SteamServerObserver.Reset();
 	}
 
 	// Here we are no longer loaded, so we need to override any DLL handles still open and unlink the DLLs.
 	UnloadSteamModules();
 }
 
-TSharedPtr<class FSteamInstanceHandler> FSteamSharedModule::ObtainSteamInstanceHandle()
+TSharedPtr<class FSteamClientInstanceHandler> FSteamSharedModule::ObtainSteamClientInstanceHandle()
 {
-	if (InstanceHandlerObserver.IsValid())
+	if (SteamClientObserver.IsValid())
 	{
-		return InstanceHandlerObserver.Pin();
+		return SteamClientObserver.Pin();
 	}
 	else
 	{
 		// Create the original base object, and store our weakptrs.
-		TSharedPtr<FSteamInstanceHandler> BaseInstance = MakeShareable(new FSteamInstanceHandler(this));
-		InstanceHandlerObserver = BaseInstance;
-		if (BaseInstance->bInitialized) // Make sure the SteamAPI was initialized properly.
+		TSharedPtr<FSteamClientInstanceHandler> BaseInstance = MakeShareable(new FSteamClientInstanceHandler(this));
+		SteamClientObserver = BaseInstance;
+		if (BaseInstance->IsInitialized()) // Make sure the SteamAPI was initialized properly.
 		{
 			// This is safe because we'll end up incrementing the BaseInstance refcounter before we go out of scope.
 			// Thus the InstanceHandler will not be deleted before we return.
-			return InstanceHandlerObserver.Pin();
+			return SteamClientObserver.Pin();
 		}
 		else
 		{
 			// We don't want to hold any references to failed instances.
-			InstanceHandlerObserver = nullptr;
+			SteamClientObserver = nullptr;
+		}
+	}
+
+	return nullptr;
+}
+
+TSharedPtr<class FSteamServerInstanceHandler> FSteamSharedModule::ObtainSteamServerInstanceHandle()
+{
+	if (SteamServerObserver.IsValid())
+	{
+		return SteamServerObserver.Pin();
+	}
+	else
+	{
+		// Create the original base object, and store our weakptrs.
+		TSharedPtr<FSteamServerInstanceHandler> BaseInstance = MakeShareable(new FSteamServerInstanceHandler(this));
+		SteamServerObserver = BaseInstance;
+		if (BaseInstance->IsInitialized()) // Make sure the SteamAPI was initialized properly.
+		{
+			// This is safe because we'll end up incrementing the BaseInstance refcounter before we go out of scope.
+			// Thus the InstanceHandler will not be deleted before we return.
+			return SteamServerObserver.Pin();
+		}
+		else
+		{
+			// We don't want to hold any references to failed instances.
+			SteamServerObserver = nullptr;
 		}
 	}
 
@@ -193,7 +244,7 @@ void FSteamSharedModule::UnloadSteamModules()
 {
 	// Only free the handles if no one is using them anymore.
 	// There's no need to check AreSteamDllsLoaded as this is done individually below.
-	if (!InstanceHandlerObserver.IsValid())
+	if (!SteamClientObserver.IsValid() && !SteamServerObserver.IsValid())
 	{
 #if LOADING_STEAM_LIBRARIES_DYNAMICALLY
 		UE_LOG(LogSteamShared, Log, TEXT("Freeing the Steam Loaded Modules..."));
@@ -213,8 +264,8 @@ void FSteamSharedModule::UnloadSteamModules()
 	}
 }
 
-FSteamInstanceHandler::FSteamInstanceHandler(FSteamSharedModule* SteamInitializer) :
-	bInitialized(false)
+FSteamClientInstanceHandler::FSteamClientInstanceHandler(FSteamSharedModule* SteamInitializer) :
+	FSteamInstanceHandlerBase()
 {
 	// A module must be loaded in order for us to initialize the Steam API.
 	if (SteamInitializer != nullptr && SteamInitializer->AreSteamDllsLoaded())
@@ -228,20 +279,108 @@ FSteamInstanceHandler::FSteamInstanceHandler(FSteamSharedModule* SteamInitialize
 
 		// The conditions mentioned in this print can be found at https://partner.steamgames.com/doc/sdk/api#initialization_and_shutdown
 		UE_LOG(LogSteamShared, Warning, TEXT("SteamAPI failed to initialize, conditions not met."));
-		bInitialized = false;
 		return;
 	}
 	UE_LOG(LogSteamShared, Warning, TEXT("SteamAPI failed to initialize as the Dlls are not loaded."));
 }
 
-FSteamInstanceHandler::~FSteamInstanceHandler()
+void FSteamClientInstanceHandler::InternalShutdown()
+{
+	UE_LOG(LogSteamShared, Log, TEXT("Unloading the Steam API..."));
+	SteamAPI_Shutdown();
+}
+
+
+FSteamServerInstanceHandler::FSteamServerInstanceHandler(FSteamSharedModule* SteamInitializer) :
+	FSteamInstanceHandlerBase()
+{
+	if (SteamInitializer == nullptr || !SteamInitializer->AreSteamDllsLoaded())
+	{
+		UE_LOG(LogSteamShared, Warning, TEXT("Steam Server API failed to initialize as the Dlls are not loaded."));
+		return;
+	}
+
+	// Get the multihome address. If there isn't one, this server will be set to listen on the any address.
+	uint32 LocalServerIP = 0;
+	FString MultiHome;
+	if (FParse::Value(FCommandLine::Get(), TEXT("MULTIHOME="), MultiHome) && !MultiHome.IsEmpty())
+	{
+		TSharedPtr<FInternetAddr> MultiHomeIP = ISocketSubsystem::Get(PLATFORM_SOCKETSUBSYSTEM)->GetAddressFromString(MultiHome);
+		if (MultiHomeIP.IsValid())
+		{
+			MultiHomeIP->GetIp(LocalServerIP);
+		}
+	}
+
+	// Grab the gameport for game communications.
+	if (FParse::Value(FCommandLine::Get(), TEXT("Port="), GamePort) == false)
+	{
+		GConfig->GetInt(TEXT("URL"), TEXT("Port"), GamePort, GEngineIni);
+	}
+
+	// Grab the SteamPort, which handles communications over the steam network.
+	SteamPort = GamePort + 1;
+
+	// Allow the command line to override the default query port for master server communications
+	if (FParse::Value(FCommandLine::Get(), TEXT("QueryPort="), QueryPort) == false)
+	{
+		if (!GConfig->GetInt(TEXT("OnlineSubsystemSteam"), TEXT("GameServerQueryPort"), QueryPort, GEngineIni))
+		{
+			QueryPort = 27015;
+		}
+	}
+
+	// Set VAC
+	bool bVACEnabled = false;
+	GConfig->GetBool(TEXT("OnlineSubsystemSteam"), TEXT("bVACEnabled"), bVACEnabled, GEngineIni);
+
+	// Set GameVersions
+	FString GameVersion;
+	GConfig->GetString(TEXT("OnlineSubsystemSteam"), TEXT("GameVersion"), GameVersion, GEngineIni);
+	if (GameVersion.Len() == 0)
+	{
+		UE_LOG(LogSteamShared, Warning, TEXT("[OnlineSubsystemSteam].GameVersion is not set. Server advertising will fail"));
+	}
+
+	UE_LOG(LogSteamShared, Verbose, TEXT("Initializing Steam Game Server IP: 0x%08X Port: %d SteamPort: %d QueryPort: %d"), LocalServerIP, GamePort, SteamPort, QueryPort);
+
+	if (SteamGameServer_Init(LocalServerIP, SteamPort, GamePort, QueryPort,
+		(bVACEnabled ? eServerModeAuthenticationAndSecure : eServerModeAuthentication),
+		TCHAR_TO_UTF8(*GameVersion)))
+	{
+		UE_LOG(LogSteamShared, Verbose, TEXT("Steam Dedicated Server API initialized."));
+		bInitialized = true;
+	}
+	else
+	{
+		UE_LOG(LogSteamShared, Warning, TEXT("Steam Dedicated Server API failed to initialize."));
+	}
+}
+
+void FSteamServerInstanceHandler::InternalShutdown()
+{
+	// Log off of the steam master servers, removes this server immediately from the backend.
+	if (SteamGameServer() && SteamGameServer()->BLoggedOn())
+	{
+		SteamGameServer()->LogOff();
+	}
+
+	// By getting here, we must be the last instance of the object, thus we should be deleted.
+	// If no one is using the API, we can shut it down.
+	UE_LOG(LogSteamShared, Log, TEXT("Unloading the Steam server API..."));
+	SteamGameServer_Shutdown();
+}
+
+bool FSteamInstanceHandlerBase::CanCleanUp() const
 {
-	// So long as a module is loaded, we can call shutdown.
-	if (bInitialized && FSteamSharedModule::IsAvailable() && FSteamSharedModule::Get().AreSteamDllsLoaded())
+	return bInitialized && FSteamSharedModule::IsAvailable() && FSteamSharedModule::Get().AreSteamDllsLoaded();
+}
+
+void FSteamInstanceHandlerBase::Destroy()
+{
+	if (CanCleanUp())
 	{
-		// By getting here, we must be the last instance of the object, thus we should be deleted.
-		// If no one is using the API, we can shut it down.
-		UE_LOG(LogSteamShared, Log, TEXT("Unloading the Steam API..."));
-		SteamAPI_Shutdown();
+		InternalShutdown();
 	}
 }
+
diff --git a/Engine/Plugins/Runtime/Steam/SteamShared/Source/SteamShared/Public/SteamSharedModule.h b/Engine/Plugins/Runtime/Steam/SteamShared/Source/SteamShared/Public/SteamSharedModule.h
--- a/Engine/Plugins/Runtime/Steam/SteamShared/Source/SteamShared/Public/SteamSharedModule.h
+++ b/Engine/Plugins/Runtime/Steam/SteamShared/Source/SteamShared/Public/SteamSharedModule.h
@@ -19,7 +19,8 @@ public:
 		SteamDLLHandle(nullptr),
 		SteamServerDLLHandle(nullptr),
 		bForceLoadSteamClientDll(false),
-		InstanceHandlerObserver(nullptr)
+		SteamClientObserver(nullptr),
+		SteamServerObserver(nullptr)
 	{
 	}
 
@@ -41,24 +42,45 @@ public:
 	 *
 	 * @return A handler to the Steam Client API, use IsValid to check if the handle is initialized.
 	 */
-	TSharedPtr<class FSteamInstanceHandler> ObtainSteamInstanceHandle();
+	TSharedPtr<class FSteamClientInstanceHandler> ObtainSteamClientInstanceHandle();
+
+	/**
+	 * Initializes Steam Server API and provides a handler that will keep the API valid for the lifetime of the
+	 * the object. Several Handlers can be active at once.
+	 *
+	 * @return A handler to the Steam Server API, use IsValid to check if the handle is initialized.
+	 */
+	TSharedPtr<class FSteamServerInstanceHandler> ObtainSteamServerInstanceHandle();
 	
 	/**
 	 * Are the Steamworks Dlls loaded
+	 *
+	 * @return if the steam dlls are currently loaded (if we are loading them dynamically, statically linked are always true)
 	 */
 	bool AreSteamDllsLoaded() const;
 	
 	/**
 	 * The path to where the Steam binaries are stored, for use in debugging.
+	 *
+	 * @return The directory path of the location of the steam dlls
 	 */
 	FString GetSteamModulePath() const;
 
 	/**
 	 * If the module will be loading the client dlls for the dedicated server instance.
 	 * Really only useful on Windows.
+	 *
+	 * @return If this application is currently loading client dlls on the server
 	 */
 	bool IsLoadingServerClientDlls() const { return bForceLoadSteamClientDll; }
 
+	/**
+	 * Checks if we can load client dlls on dedicated server instances.
+	 * 
+	 * @return On dedicated servers on windows, this returns true, for other platforms this returns false (as feature is unnecessary)
+	 */
+	bool CanLoadClientDllsOnServer() const;
+
 	/**
 	 * Singleton-like access to this module's interface.  This is just for convenience!
 	 * Beware of calling this during the shutdown phase, though.  Your module might have been unloaded already.
@@ -92,7 +114,8 @@ private:
 	bool bForceLoadSteamClientDll;
 
 	/** Object that holds the refcounted pointer that's given out */
-	TWeakPtr<class FSteamInstanceHandler> InstanceHandlerObserver;
+	TWeakPtr<class FSteamClientInstanceHandler> SteamClientObserver;
+	TWeakPtr<class FSteamServerInstanceHandler> SteamServerObserver;
 
 	/** Load the required modules for Steam */
 	void LoadSteamModules();
@@ -101,21 +124,68 @@ private:
 	void UnloadSteamModules();
 };
 
-/** A simple instance handler that creates and uninitializes the SteamAPI automatically. */
-class STEAMSHARED_API FSteamInstanceHandler
+/** Base instance handler class for the Steam shared classes, this allows for less code redundancy between the shared modules. */
+class STEAMSHARED_API FSteamInstanceHandlerBase
 {
 public:
-	virtual ~FSteamInstanceHandler();
+	virtual ~FSteamInstanceHandlerBase() {}
+	virtual bool IsInitialized() const { return bInitialized; }
+
+protected:
+	FSteamInstanceHandlerBase() : bInitialized(false) {}
 
-PACKAGE_SCOPE:
-	/** This is mostly here for safety so that we don't try to clean up when we shouldn't */
 	bool bInitialized;
+	virtual bool CanCleanUp() const;
+	virtual void Destroy();
+	virtual void InternalShutdown() = 0;
+};
 
+/** A simple instance handler that creates and uninitializes the client SteamAPI automatically. */
+class STEAMSHARED_API FSteamClientInstanceHandler : public FSteamInstanceHandlerBase
+{
+public:
+	virtual ~FSteamClientInstanceHandler() { Destroy(); }
+
+PACKAGE_SCOPE:
 	/** Initializes the Steamworks client API on call */	
-	FSteamInstanceHandler(FSteamSharedModule* SteamInitializer);
+	FSteamClientInstanceHandler(FSteamSharedModule* SteamInitializer);
+
+protected:
+	virtual void InternalShutdown() override;
 
 private:
-	FSteamInstanceHandler() : bInitialized(false)
+	FSteamClientInstanceHandler() : 
+		FSteamInstanceHandlerBase()
 	{
 	}
 };
+
+/** A simple instance handler that creates and uninitializes the server SteamAPI automatically. */
+class STEAMSHARED_API FSteamServerInstanceHandler : public FSteamInstanceHandlerBase
+{
+public:
+	virtual ~FSteamServerInstanceHandler() { Destroy(); }
+
+	int32 GetGamePort() const { return GamePort; }
+	int32 GetSteamPort() const { return SteamPort; }
+	int32 GetQueryPort() const { return QueryPort; }
+
+PACKAGE_SCOPE:
+	/** Initializes the Steamworks server API on call */
+	FSteamServerInstanceHandler(FSteamSharedModule* SteamInitializer);
+
+protected:
+	int32 GamePort;
+	int32 SteamPort;
+	int32 QueryPort;
+	virtual void InternalShutdown() override;
+
+private:
+	FSteamServerInstanceHandler() : 
+		FSteamInstanceHandlerBase(),
+		GamePort(-1),
+		SteamPort(-1),
+		QueryPort(-1)
+	{
+	}
+};
\ No newline at end of file
diff --git a/Engine/Plugins/Runtime/Steam/SteamShared/Source/SteamShared/SteamShared.Build.cs b/Engine/Plugins/Runtime/Steam/SteamShared/Source/SteamShared/SteamShared.Build.cs
--- a/Engine/Plugins/Runtime/Steam/SteamShared/Source/SteamShared/SteamShared.Build.cs
+++ b/Engine/Plugins/Runtime/Steam/SteamShared/Source/SteamShared/SteamShared.Build.cs
@@ -15,8 +15,13 @@ public class SteamShared : ModuleRules
         {
             PublicDefinitions.Add("STEAM_SDK_INSTALLED");
         }
-
-        PrivateDependencyModuleNames.Add("Core");
+	
+		PrivateDependencyModuleNames.AddRange(
+			new string[] {
+				"Core", 
+				"Sockets"
+            }
+		);
 
 		AddEngineThirdPartyPrivateStaticDependencies(Target, "Steamworks");
 	}
