Adding includetool override to remove CIS warnings for Steam package headers

#rnx

diff --git a/Engine/Plugins/Runtime/Steam/SteamShared/Source/SteamShared/Public/SteamSharedPackage.h b/Engine/Plugins/Runtime/Steam/SteamShared/Source/SteamShared/Public/SteamSharedPackage.h
--- a/Engine/Plugins/Runtime/Steam/SteamShared/Source/SteamShared/Public/SteamSharedPackage.h
+++ b/Engine/Plugins/Runtime/Steam/SteamShared/Source/SteamShared/Public/SteamSharedPackage.h
@@ -5,6 +5,8 @@
 // Intended to be the last include in an exported class definition
 // Properly defines some members as "public to the module" vs "private to the consumer/user"
 
+// [[ IncludeTool: Inline ]]
+
 #undef PACKAGE_SCOPE
 #ifdef STEAMSHARED_PACKAGE
 #define PACKAGE_SCOPE public
