Fixed steam_appid.txt not being generated for dedicated servers and improved error logging for steam initialization.
#Synced-CL 46989725

diff --git a/Engine/Plugins/Runtime/Steam/SteamShared/Source/SteamShared/Private/SteamSharedModule.cpp b/Engine/Plugins/Runtime/Steam/SteamShared/Source/SteamShared/Private/SteamSharedModule.cpp
--- a/Engine/Plugins/Runtime/Steam/SteamShared/Source/SteamShared/Private/SteamSharedModule.cpp
+++ b/Engine/Plugins/Runtime/Steam/SteamShared/Source/SteamShared/Private/SteamSharedModule.cpp
@@ -342,8 +342,8 @@ FSteamClientInstanceHandler::FSteamClientInstanceHandler(FSteamSharedModule* Ste
 			UE_LOG(LogSteamShared, Warning, TEXT("Failed to create steam_appid.txt. If the file doesn't exist, SteamAPI_Init will fail."));
 		}
 #endif // !UE_BUILD_SHIPPING && !UE_BUILD_SHIPPING_WITH_EDITOR 
-
-		if (SteamAPI_Init())
+		SteamErrMsg InitServerErrorMsg;
+		if (SteamAPI_InitEx(&InitServerErrorMsg) == k_ESteamAPIInitResult_OK)
 		{
 			UE_LOG(LogSteamShared, Verbose, TEXT("SteamAPI initialized"));
 			bInitialized = true;
@@ -351,7 +351,7 @@ FSteamClientInstanceHandler::FSteamClientInstanceHandler(FSteamSharedModule* Ste
 		}
 
 		// The conditions mentioned in this print can be found at https://partner.steamgames.com/doc/sdk/api#initialization_and_shutdown
-		UE_LOG(LogSteamShared, Warning, TEXT("SteamAPI failed to initialize, conditions not met."));
+		UE_LOG(LogSteamShared, Warning, TEXT("SteamAPI failed to initialize: %s"), StringCast<TCHAR>(InitServerErrorMsg).Get());
 		return;
 	}
 	UE_LOG(LogSteamShared, Warning, TEXT("SteamAPI failed to initialize as the Dlls are not loaded."));
@@ -408,16 +408,25 @@ FSteamServerInstanceHandler::FSteamServerInstanceHandler(FSteamSharedModule* Ste
 
 	UE_LOG(LogSteamShared, Verbose, TEXT("Initializing Steam Game Server IP: 0x%08X Port: %d QueryPort: %d"), LocalServerIP, GamePort, QueryPort);
 
-	if (SteamGameServer_Init(LocalServerIP, GamePort, QueryPort,
+#if !UE_BUILD_SHIPPING && !UE_BUILD_SHIPPING_WITH_EDITOR
+	if (!WriteSteamAppIdToDisk())
+	{
+		UE_LOG(LogSteamShared, Warning, TEXT("Failed to create steam_appid.txt. If the file doesn't exist, SteamAPI_Init will fail."));
+	}
+#endif // !UE_BUILD_SHIPPING && !UE_BUILD_SHIPPING_WITH_EDITOR 
+	
+	SteamErrMsg InitServerErrorMsg;
+	ESteamAPIInitResult InitStatus = SteamGameServer_InitEx(LocalServerIP, GamePort, QueryPort,
 		(bVACEnabled ? eServerModeAuthenticationAndSecure : eServerModeAuthentication),
-		TCHAR_TO_UTF8(*GameVersion)))
+		TCHAR_TO_UTF8(*GameVersion), &InitServerErrorMsg);
+	if (InitStatus == k_ESteamAPIInitResult_OK)
 	{
 		UE_LOG(LogSteamShared, Verbose, TEXT("Steam Dedicated Server API initialized."));
 		bInitialized = true;
 	}
 	else
 	{
-		UE_LOG(LogSteamShared, Warning, TEXT("Steam Dedicated Server API failed to initialize."));
+		UE_LOG(LogSteamShared, Warning, TEXT("Steam Dedicated Server API failed to initialize: %s"), StringCast<TCHAR>(InitServerErrorMsg).Get());
 	}
 }
 
